fix OodleHandlerComponent turn off verbose logging

diff --git a/Engine/Plugins/Compression/OodleNetwork/Source/Private/OodleHandlerComponent.cpp b/Engine/Plugins/Compression/OodleNetwork/Source/Private/OodleHandlerComponent.cpp
--- a/Engine/Plugins/Compression/OodleNetwork/Source/Private/OodleHandlerComponent.cpp
+++ b/Engine/Plugins/Compression/OodleNetwork/Source/Private/OodleHandlerComponent.cpp
@@ -25,8 +25,7 @@
 DEFINE_LOG_CATEGORY(OodleHandlerComponentLog);
 
 #ifndef OODLE_HANDLER_VERBOSE_LOG
-//#define OODLE_HANDLER_VERBOSE_LOG	0
-#define OODLE_HANDLER_VERBOSE_LOG	1
+#define OODLE_HANDLER_VERBOSE_LOG	0
 #endif
 
 #ifndef OODLE_USE_FALLBACK_DICTIONARY
@@ -1406,14 +1405,14 @@ void OodleHandlerComponent::NotifyAnalyticsProvider()
 static bool OodleExec(UWorld* InWorld, const TCHAR* Cmd, FOutputDevice& Ar)
 {
 
-	#if OODLE_HANDLER_VERBOSE_LOG
-	UE_LOG(OodleHandlerComponentLog, Log, TEXT("OodleExec") );
-	#endif
-
 	bool bReturnVal = false;
 
 	if (FParse::Command(&Cmd, TEXT("Oodle")))
 	{
+		#if OODLE_HANDLER_VERBOSE_LOG
+		UE_LOG(OodleHandlerComponentLog, Log, TEXT("OodleExec") );
+		#endif
+
 		// Used by unit testing code, to enable/disable Oodle during a unit test
 		// NOTE: Do not use while a NetConnection is using Oodle, as this will cause it to break. Debugging/Testing only.
 		if (FParse::Command(&Cmd, TEXT("ForceEnable")))
