check return value of OodleNetwork1UDP_State_Uncompact

diff --git a/Engine/Plugins/Compression/OodleNetwork/Source/Private/OodleNetworkHandlerComponent.cpp b/Engine/Plugins/Compression/OodleNetwork/Source/Private/OodleNetworkHandlerComponent.cpp
--- a/Engine/Plugins/Compression/OodleNetwork/Source/Private/OodleNetworkHandlerComponent.cpp
+++ b/Engine/Plugins/Compression/OodleNetwork/Source/Private/OodleNetworkHandlerComponent.cpp
@@ -603,8 +603,15 @@ void OodleNetworkHandlerComponent::InitializeDictionary(FString FilePath, TShare
 				uint32 CompressorStateSize = OodleNetwork1UDP_State_Size();
 				OodleNetwork1UDP_State* CompressorState = (OodleNetwork1UDP_State*)FMemory::Malloc(CompressorStateSize);
 
-				OodleNetwork1UDP_State_Uncompact(CompressorState, (OodleNetwork1UDP_StateCompacted*)CompactCompressorState);
-
+				OO_BOOL UncompactOk = OodleNetwork1UDP_State_Uncompact(CompressorState, (OodleNetwork1UDP_StateCompacted*)CompactCompressorState);
+				if ( ! UncompactOk )
+				{
+					UE_LOG(OodleNetworkHandlerComponentLog, Error, TEXT("OodleNetwork1UDP_State_Uncompact failed!"));
+					// @todo Oodle does soft-fail like this work?
+					//	we'd like to have failures just disable OodleNetwork and allow work to continue
+					bEnableOodle = false;
+					return;
+				}
 
 				// Create the shared dictionary state
 				int32 HashTableSize = BoundArc.Header.HashTableSize.Get();
