Enable loading of old Oodle dictionary files as well as writing old version when training off of packet captures.

diff --git a/Engine/Plugins/Compression/OodleNetwork/Source/Classes/OodleNetworkTrainerCommandlet.h b/Engine/Plugins/Compression/OodleNetwork/Source/Classes/OodleNetworkTrainerCommandlet.h
--- a/Engine/Plugins/Compression/OodleNetwork/Source/Classes/OodleNetworkTrainerCommandlet.h
+++ b/Engine/Plugins/Compression/OodleNetwork/Source/Classes/OodleNetworkTrainerCommandlet.h
@@ -70,6 +70,10 @@ public:
 	UPROPERTY(config)
 	bool bCompressionTest;
 
+	/** Whether or not to write oodle version 5 dictionaries for back compat */
+	UPROPERTY(config)
+	bool bWriteV5Dictionaries;
+
 	/** Size of the hash table to use for the dictionary */
 	UPROPERTY(config)
 	int32 HashTableSize;
@@ -208,6 +212,9 @@ private:
 	/** Whether or not compression testing should be performed after dictionary generation (uses up some of the packets) */
 	bool bCompressionTest;
 
+	/** Whether or not to write oodle version 5 dictionaries for back compat */
+	bool bWriteV5Dictionaries;
+
 	/** Size of the hash table to use for the dictionary */
 	int32 HashTableSize;
 
diff --git a/Engine/Plugins/Compression/OodleNetwork/Source/Private/OodleNetworkArchives.cpp b/Engine/Plugins/Compression/OodleNetwork/Source/Private/OodleNetworkArchives.cpp
--- a/Engine/Plugins/Compression/OodleNetwork/Source/Private/OodleNetworkArchives.cpp
+++ b/Engine/Plugins/Compression/OodleNetwork/Source/Private/OodleNetworkArchives.cpp
@@ -366,11 +366,12 @@ FOodleNetworkDictionaryArchive::FOodleNetworkDictionaryArchive(FArchive& InInner
 	}
 }
 
-void FOodleNetworkDictionaryArchive::SetDictionaryHeaderValues(int32 InHashTableSize)
+void FOodleNetworkDictionaryArchive::SetDictionaryHeaderValues(int32 InHashTableSize, uint32 InWritingOodleVersion)
 {
 	check(IsSaving());
 
 	Header.HashTableSize = InHashTableSize;
+	Header.OodleMajorHeaderVersion = InWritingOodleVersion;
 }
 
 void FOodleNetworkDictionaryArchive::SerializeHeader()
diff --git a/Engine/Plugins/Compression/OodleNetwork/Source/Private/OodleNetworkHandlerComponent.cpp b/Engine/Plugins/Compression/OodleNetwork/Source/Private/OodleNetworkHandlerComponent.cpp
--- a/Engine/Plugins/Compression/OodleNetwork/Source/Private/OodleNetworkHandlerComponent.cpp
+++ b/Engine/Plugins/Compression/OodleNetwork/Source/Private/OodleNetworkHandlerComponent.cpp
@@ -603,7 +603,7 @@ void OodleNetworkHandlerComponent::InitializeDictionary(FString FilePath, TShare
 				uint32 CompressorStateSize = OodleNetwork1UDP_State_Size();
 				OodleNetwork1UDP_State* CompressorState = (OodleNetwork1UDP_State*)FMemory::Malloc(CompressorStateSize);
 
-				OO_BOOL UncompactOk = OodleNetwork1UDP_State_Uncompact(CompressorState, (OodleNetwork1UDP_StateCompacted*)CompactCompressorState);
+				OO_BOOL UncompactOk = OodleNetwork1UDP_State_Uncompact_ForVersion(CompressorState, (OodleNetwork1UDP_StateCompacted*)CompactCompressorState, BoundArc.Header.OodleMajorHeaderVersion);
 				if ( ! UncompactOk )
 				{
 					UE_LOG(OodleNetworkHandlerComponentLog, Error, TEXT("OodleNetwork1UDP_State_Uncompact failed!"));
diff --git a/Engine/Plugins/Compression/OodleNetwork/Source/Private/OodleNetworkTrainerCommandlet.cpp b/Engine/Plugins/Compression/OodleNetwork/Source/Private/OodleNetworkTrainerCommandlet.cpp
--- a/Engine/Plugins/Compression/OodleNetwork/Source/Private/OodleNetworkTrainerCommandlet.cpp
+++ b/Engine/Plugins/Compression/OodleNetwork/Source/Private/OodleNetworkTrainerCommandlet.cpp
@@ -636,13 +636,16 @@ bool FOodleNetworkDictionaryGenerator::InitGenerator()
 	TrialGenerations = CDO->TrialGenerations;
 	bNoTrials = CDO->bNoTrials;
 	bCompressionTest = CDO->bCompressionTest;
+	bWriteV5Dictionaries = CDO->bWriteV5Dictionaries;
 
 	// Commandline overrides
 	FParse::Value(FCommandLine::Get(), TEXT("-HashTableSize="), HashTableSize);
 	FParse::Value(FCommandLine::Get(), TEXT("-DictionarySize="), DictionarySize);
 	FParse::Value(FCommandLine::Get(), TEXT("-DictionaryTrials="), DictionaryTrials);
 	FParse::Value(FCommandLine::Get(), TEXT("-TrialRandomness="), TrialRandomness);
-	FParse::Value(FCommandLine::Get(), TEXT("-TrialGenerations="), TrialGenerations);
+	FParse::Value(FCommandLine::Get(), TEXT("-TrialGenerations="), TrialGenerations);	
+	FParse::Bool(FCommandLine::Get(), TEXT("-WriteV5Dictionaries="), bWriteV5Dictionaries);
+
 	bNoTrials = FParse::Param(FCommandLine::Get(), TEXT("NoTrials"));
 	bCompressionTest = FParse::Param(FCommandLine::Get(), TEXT("CompressionTest"));
 
@@ -711,6 +714,8 @@ bool FOodleNetworkDictionaryGenerator::InitGenerator()
 		UE_LOG(OodleNetworkHandlerComponentLog, Log, TEXT("- TrialGenerations: %i"), TrialGenerations);
 	}
 
+	UE_LOG(OodleNetworkHandlerComponentLog, Log, TEXT("- WriteDictionaryVersion: %i"), bWriteV5Dictionaries ? 5 : OODLE2NET_VERSION_MAJOR);
+
 
 
 	// Adjust non-absolute directories, so they are relative to the game directory
@@ -1004,8 +1009,19 @@ bool FOodleNetworkDictionaryGenerator::GenerateAndWriteDictionary()
 	// Now compact the dictionary and initial state data, in preparation for writing to file
 	const SSIZE_T CompactState_Size = OodleNetwork1UDP_StateCompacted_MaxSize();
 	OodleNetwork1UDP_StateCompacted* CompactCompressorState = (OodleNetwork1UDP_StateCompacted*)FMemory::Malloc(CompactState_Size);
-	SSIZE_T CompactCompressorStateBytes = OodleNetwork1UDP_State_Compact(CompactCompressorState, CompressorState);
 
+	uint32 CompactWriteVersion = 0;
+	SSIZE_T CompactCompressorStateBytes = 0;
+	if (bWriteV5Dictionaries)
+	{
+		CompactCompressorStateBytes  = OodleNetwork1UDP_State_Compact_ForVersion(CompactCompressorState, CompressorState, 5);
+		CompactWriteVersion = 5;
+	}
+	else
+	{
+		CompactCompressorStateBytes = OodleNetwork1UDP_State_Compact(CompactCompressorState, CompressorState);
+		CompactWriteVersion = OODLE2NET_VERSION_MAJOR;
+	}
 
 	// Now write to the final file
 	FArchive* OutArc = IFileManager::Get().CreateFileWriter(*OutputDictionaryFile);
@@ -1017,7 +1033,7 @@ bool FOodleNetworkDictionaryGenerator::GenerateAndWriteDictionary()
 		{
 			FOodleNetworkDictionaryArchive OutputFile(*OutArc);
 
-			OutputFile.SetDictionaryHeaderValues(HashTableSize);
+			OutputFile.SetDictionaryHeaderValues(HashTableSize, CompactWriteVersion);
 			OutputFile.SerializeHeader();
 
 			bSuccess = !OutputFile.IsError();
diff --git a/Engine/Plugins/Compression/OodleNetwork/Source/Public/OodleNetworkArchives.h b/Engine/Plugins/Compression/OodleNetwork/Source/Public/OodleNetworkArchives.h
--- a/Engine/Plugins/Compression/OodleNetwork/Source/Public/OodleNetworkArchives.h
+++ b/Engine/Plugins/Compression/OodleNetwork/Source/Public/OodleNetworkArchives.h
@@ -426,8 +426,10 @@ public:
 	 * Sets dictionary header values, that should be set prior to serializing the header to file
 	 *
 	 * @param InHashTableSize	The HashTableSize to set in the header
+	 * @param InWritingOodleVersion	The compatibility version we are writing.
+	 * 
 	 */
-	void SetDictionaryHeaderValues(int32 InHashTableSize);
+	void SetDictionaryHeaderValues(int32 InHashTableSize, uint32 InWritingOodleVersion);
 
 	/**
 	 * Serializes initial basic header values
