Integrating github pull request to fix compiling in shipping with stats enabled.

#8704

diff --git a/Engine/Plugins/Compression/OodleNetwork/Source/Private/OodleNetworkHandlerComponent.cpp b/Engine/Plugins/Compression/OodleNetwork/Source/Private/OodleNetworkHandlerComponent.cpp
--- a/Engine/Plugins/Compression/OodleNetwork/Source/Private/OodleNetworkHandlerComponent.cpp
+++ b/Engine/Plugins/Compression/OodleNetwork/Source/Private/OodleNetworkHandlerComponent.cpp
@@ -48,7 +48,9 @@ DEFINE_LOG_CATEGORY(OodleNetworkHandlerComponentLog);
 
 
 #if STATS
+#if !UE_BUILD_SHIPPING
 DEFINE_STAT(STAT_PacketReservedOodle);
+#endif
 
 DEFINE_STAT(STAT_Oodle_OutRaw);
 DEFINE_STAT(STAT_Oodle_OutCompressed);
