remove HoloLens from engine plugins' PlatformDenyList in favor of platform extension.

#rnx

diff --git a/Engine/Platforms/Hololens/Plugins/Compression/OodleNetwork/OodleNetwork_HoloLens.uplugin b/Engine/Platforms/Hololens/Plugins/Compression/OodleNetwork/OodleNetwork_HoloLens.uplugin
new file mode 100644
--- /dev/null
+++ b/Engine/Platforms/Hololens/Plugins/Compression/OodleNetwork/OodleNetwork_HoloLens.uplugin
@@ -0,0 +1,13 @@
+{
+	"FileVersion": 3,
+	"bIsPluginExtension": true,
+	"Modules": [
+		{
+			"Name": "OodleNetworkHandlerComponent",
+			"Type": "Runtime",
+			"PlatformDenyList": [
+				"HoloLens"
+			]
+		}
+	]
+}
\ No newline at end of file
diff --git a/Engine/Platforms/Hololens/Plugins/Compression/OodleNetwork/Source/OodleNetworkHandlerComponent_HoloLens.Build.cs b/Engine/Platforms/Hololens/Plugins/Compression/OodleNetwork/Source/OodleNetworkHandlerComponent_HoloLens.Build.cs
new file mode 100644
--- /dev/null
+++ b/Engine/Platforms/Hololens/Plugins/Compression/OodleNetwork/Source/OodleNetworkHandlerComponent_HoloLens.Build.cs
@@ -0,0 +1,12 @@
+// Copyright Epic Games, Inc. All Rights Reserved.
+
+namespace UnrealBuildTool.Rules
+{
+	public class OodleNetworkHandlerComponent_HoloLens : OodleNetworkHandlerComponent
+	{
+		public OodleNetworkHandlerComponent_HoloLens(ReadOnlyTargetRules Target) : base(Target)
+		{
+			throw new System.Exception("Should not be here - HoloLens in the uplugin's PlatformDenyList");
+		}
+	}
+}
diff --git a/Engine/Platforms/Hololens/Plugins/Media/BinkMedia/BinkMedia_HoloLens.uplugin b/Engine/Platforms/Hololens/Plugins/Media/BinkMedia/BinkMedia_HoloLens.uplugin
new file mode 100644
--- /dev/null
+++ b/Engine/Platforms/Hololens/Plugins/Media/BinkMedia/BinkMedia_HoloLens.uplugin
@@ -0,0 +1,13 @@
+{
+	"FileVersion": 3,
+	"bIsPluginExtension": true,
+	"Modules": [
+		{
+			"Name": "BinkMediaPlayer",
+			"Type": "RuntimeNoCommandlet",
+			"PlatformDenyList": [
+				"HoloLens"
+			]
+		}
+	]
+}
\ No newline at end of file
diff --git a/Engine/Platforms/Hololens/Plugins/Media/BinkMedia/Source/BinkMediaPlayer/BinkMediaPlayer_HoloLens.Build.cs b/Engine/Platforms/Hololens/Plugins/Media/BinkMedia/Source/BinkMediaPlayer/BinkMediaPlayer_HoloLens.Build.cs
new file mode 100644
--- /dev/null
+++ b/Engine/Platforms/Hololens/Plugins/Media/BinkMedia/Source/BinkMediaPlayer/BinkMediaPlayer_HoloLens.Build.cs
@@ -0,0 +1,12 @@
+// Copyright Epic Games, Inc. All Rights Reserved.
+
+namespace UnrealBuildTool.Rules
+{
+	public class BinkMediaPlayer_HoloLens : BinkMediaPlayer
+	{
+		public BinkMediaPlayer_HoloLens(ReadOnlyTargetRules Target) : base(Target)
+		{
+			throw new System.Exception("Should not be here - HoloLens in the uplugin's PlatformDenyList");
+		}
+	}
+}
diff --git a/Engine/Platforms/Hololens/Plugins/Online/Experimental/OnlineServicesEOS/OnlineServicesEOS_HoloLens.uplugin b/Engine/Platforms/Hololens/Plugins/Online/Experimental/OnlineServicesEOS/OnlineServicesEOS_HoloLens.uplugin
new file mode 100644
--- /dev/null
+++ b/Engine/Platforms/Hololens/Plugins/Online/Experimental/OnlineServicesEOS/OnlineServicesEOS_HoloLens.uplugin
@@ -0,0 +1,13 @@
+{
+	"FileVersion": 3,
+	"bIsPluginExtension": true,
+	"Modules": [
+		{
+			"Name": "OnlineServicesEOS",
+			"Type": "Runtime",
+			"PlatformDenyList": [
+				"HoloLens"
+			]
+		}
+	]
+}
\ No newline at end of file
diff --git a/Engine/Platforms/Hololens/Plugins/Online/Experimental/OnlineServicesEOS/Source/OnlineServicesEOS_HoloLens.Build.cs b/Engine/Platforms/Hololens/Plugins/Online/Experimental/OnlineServicesEOS/Source/OnlineServicesEOS_HoloLens.Build.cs
new file mode 100644
--- /dev/null
+++ b/Engine/Platforms/Hololens/Plugins/Online/Experimental/OnlineServicesEOS/Source/OnlineServicesEOS_HoloLens.Build.cs
@@ -0,0 +1,12 @@
+// Copyright Epic Games, Inc. All Rights Reserved.
+
+namespace UnrealBuildTool.Rules
+{
+	public class OnlineServicesEOS_HoloLens : OnlineServicesEOS
+	{
+		public OnlineServicesEOS_HoloLens(ReadOnlyTargetRules Target) : base(Target)
+		{
+			throw new System.Exception("Should not be here - HoloLens in the uplugin's PlatformDenyList");
+		}
+	}
+}
diff --git a/Engine/Platforms/Hololens/Plugins/Online/Experimental/OnlineServicesEOSGS/OnlineServicesEOSGS_HoloLens.uplugin b/Engine/Platforms/Hololens/Plugins/Online/Experimental/OnlineServicesEOSGS/OnlineServicesEOSGS_HoloLens.uplugin
new file mode 100644
--- /dev/null
+++ b/Engine/Platforms/Hololens/Plugins/Online/Experimental/OnlineServicesEOSGS/OnlineServicesEOSGS_HoloLens.uplugin
@@ -0,0 +1,13 @@
+{
+	"FileVersion": 3,
+	"bIsPluginExtension": true,
+	"Modules": [
+		{
+			"Name": "OnlineServicesEOSGS",
+			"Type": "Runtime",
+			"PlatformDenyList": [
+				"HoloLens"
+			]
+		}
+	]
+}
\ No newline at end of file
diff --git a/Engine/Platforms/Hololens/Plugins/Online/Experimental/OnlineServicesEOSGS/Source/OnlineServicesEOSGS_HoloLens.Build.cs b/Engine/Platforms/Hololens/Plugins/Online/Experimental/OnlineServicesEOSGS/Source/OnlineServicesEOSGS_HoloLens.Build.cs
new file mode 100644
--- /dev/null
+++ b/Engine/Platforms/Hololens/Plugins/Online/Experimental/OnlineServicesEOSGS/Source/OnlineServicesEOSGS_HoloLens.Build.cs
@@ -0,0 +1,12 @@
+// Copyright Epic Games, Inc. All Rights Reserved.
+
+namespace UnrealBuildTool.Rules
+{
+	public class OnlineServicesEOSGS_HoloLens : OnlineServicesEOSGS
+	{
+		public OnlineServicesEOSGS_HoloLens(ReadOnlyTargetRules Target) : base(Target)
+		{
+			throw new System.Exception("Should not be here - HoloLens in the uplugin's PlatformDenyList");
+		}
+	}
+}
diff --git a/Engine/Platforms/Hololens/Plugins/Runtime/ApexDestruction/ApexDestruction_HoloLens.uplugin b/Engine/Platforms/Hololens/Plugins/Runtime/ApexDestruction/ApexDestruction_HoloLens.uplugin
new file mode 100644
--- /dev/null
+++ b/Engine/Platforms/Hololens/Plugins/Runtime/ApexDestruction/ApexDestruction_HoloLens.uplugin
@@ -0,0 +1,13 @@
+{
+	"FileVersion": 3,
+	"bIsPluginExtension": true,
+	"Modules": [
+		{
+			"Name": "ApexDestruction",
+			"Type": "Runtime",
+			"PlatformDenyList": [
+				"HoloLens"
+			]
+		}
+	]
+}
\ No newline at end of file
diff --git a/Engine/Platforms/Hololens/Plugins/Runtime/ApexDestruction/Source/ApexDestruction/ApexDestruction_HoloLens.Build.cs b/Engine/Platforms/Hololens/Plugins/Runtime/ApexDestruction/Source/ApexDestruction/ApexDestruction_HoloLens.Build.cs
new file mode 100644
--- /dev/null
+++ b/Engine/Platforms/Hololens/Plugins/Runtime/ApexDestruction/Source/ApexDestruction/ApexDestruction_HoloLens.Build.cs
@@ -0,0 +1,12 @@
+// Copyright Epic Games, Inc. All Rights Reserved.
+
+namespace UnrealBuildTool.Rules
+{
+	public class ApexDestruction_HoloLens : ApexDestruction
+	{
+		public ApexDestruction_HoloLens(ReadOnlyTargetRules Target) : base(Target)
+		{
+			throw new System.Exception("Should not be here - HoloLens in the uplugin's PlatformDenyList");
+		}
+	}
+}
diff --git a/Engine/Plugins/Compression/OodleNetwork/OodleNetwork.uplugin b/Engine/Plugins/Compression/OodleNetwork/OodleNetwork.uplugin
--- a/Engine/Plugins/Compression/OodleNetwork/OodleNetwork.uplugin
+++ b/Engine/Plugins/Compression/OodleNetwork/OodleNetwork.uplugin
@@ -17,10 +17,7 @@
     {
       "Name": "OodleNetworkHandlerComponent",
       "Type": "Runtime",
-      "LoadingPhase": "Default",
-	  "PlatformDenyList": [ 
-		"HoloLens"
-      ]
+      "LoadingPhase": "Default"
     }
   ]
-}
+}
\ No newline at end of file
diff --git a/Engine/Plugins/Media/BinkMedia/BinkMedia.uplugin b/Engine/Plugins/Media/BinkMedia/BinkMedia.uplugin
--- a/Engine/Plugins/Media/BinkMedia/BinkMedia.uplugin
+++ b/Engine/Plugins/Media/BinkMedia/BinkMedia.uplugin
@@ -22,7 +22,6 @@
 			"LoadingPhase" : "PreLoadingScreen",
 			"TargetDenyList": [ "Server" ],
 			"PlatformDenyList": [ 
-				"HoloLens",
 				"tvOS"
 			]
 		},
@@ -34,4 +33,4 @@
 			]
 		}
     ]
-}
+}
\ No newline at end of file
diff --git a/Engine/Plugins/Online/Experimental/OnlineServicesEOS/OnlineServicesEOS.uplugin b/Engine/Plugins/Online/Experimental/OnlineServicesEOS/OnlineServicesEOS.uplugin
--- a/Engine/Plugins/Online/Experimental/OnlineServicesEOS/OnlineServicesEOS.uplugin
+++ b/Engine/Plugins/Online/Experimental/OnlineServicesEOS/OnlineServicesEOS.uplugin
@@ -17,7 +17,6 @@
 			"LoadingPhase": "PostConfigInit",
 			"PlatformDenyList": [
 				"LinuxArm64",
-				"Hololens",
 				"TVOS"
 			]
 		}
@@ -42,4 +41,4 @@
 			"Enabled": true
 		}
 	]
-}
+}
\ No newline at end of file
diff --git a/Engine/Plugins/Online/Experimental/OnlineServicesEOSGS/OnlineServicesEOSGS.uplugin b/Engine/Plugins/Online/Experimental/OnlineServicesEOSGS/OnlineServicesEOSGS.uplugin
--- a/Engine/Plugins/Online/Experimental/OnlineServicesEOSGS/OnlineServicesEOSGS.uplugin
+++ b/Engine/Plugins/Online/Experimental/OnlineServicesEOSGS/OnlineServicesEOSGS.uplugin
@@ -17,7 +17,6 @@
 			"LoadingPhase": "PostConfigInit",
 			"PlatformDenyList": [
 				"LinuxArm64",
-				"Hololens",
 				"TVOS"
 			]
 		}
@@ -42,4 +41,4 @@
 			"Enabled": true
 		}
 	]
-}
+}
\ No newline at end of file
diff --git a/Engine/Plugins/Runtime/ApexDestruction/ApexDestruction.uplugin b/Engine/Plugins/Runtime/ApexDestruction/ApexDestruction.uplugin
--- a/Engine/Plugins/Runtime/ApexDestruction/ApexDestruction.uplugin
+++ b/Engine/Plugins/Runtime/ApexDestruction/ApexDestruction.uplugin
@@ -19,10 +19,7 @@
 		{
 			"Name": "ApexDestruction",
 			"Type": "Runtime",
-			"LoadingPhase": "PreDefault",
-			"PlatformDenyList": [
-				"HoloLens"
-			]
+			"LoadingPhase": "PreDefault"
 		},
 		{
 			"Name": "ApexDestructionEditor",
@@ -30,4 +27,4 @@
 			"LoadingPhase": "PreDefault"
 		}
 	]
-}
+}
\ No newline at end of file
