PacketHandler/Oodle profiling tweaks.

diff --git a/Engine/Plugins/Compression/OodleNetwork/Source/Private/OodleNetworkHandlerComponent.cpp b/Engine/Plugins/Compression/OodleNetwork/Source/Private/OodleNetworkHandlerComponent.cpp
--- a/Engine/Plugins/Compression/OodleNetwork/Source/Private/OodleNetworkHandlerComponent.cpp
+++ b/Engine/Plugins/Compression/OodleNetwork/Source/Private/OodleNetworkHandlerComponent.cpp
@@ -103,8 +103,11 @@ static TArray<OodleNetworkHandlerComponent*> OodleComponentList;
  * CVars
  */
 
-TAutoConsoleVariable<int32> CVarOodleMinSizeForCompression(
-	TEXT("net.OodleMinSizeForCompression"), 0,
+static int32 GOodleMinSizeForCompression = 0;
+
+FAutoConsoleVariableRef CVarOodleMinSizeForCompression(
+	TEXT("net.OodleMinSizeForCompression"),
+	GOodleMinSizeForCompression,
 	TEXT("The minimum size an outgoing packet must be, for it to be considered for compression (does not count overhead of handler components which process packets after Oodle)."));
 
 
@@ -572,6 +575,8 @@ void OodleNetworkHandlerComponent::InitializeDictionaries()
 
 void OodleNetworkHandlerComponent::RemoteInitializeDictionaries()
 {
+	QUICK_SCOPE_CYCLE_COUNTER(STAT_Oodle_RemoteInitializeDictionaries);
+
 	InitializeDictionaries();
 
 	if (bOodleNetworkAnalytics && NetAnalyticsData.IsValid())
@@ -1060,6 +1065,12 @@ void OodleNetworkHandlerComponent::Incoming(FIncomingPacketRef PacketRef)
 #if !UE_BUILD_SHIPPING || OODLE_DEV_SHIPPING
 						if (bCaptureMode && Handler->Mode == Handler::Mode::Server && InPacketLog != nullptr)
 						{
+#if !UE_BUILD_SHIPPING
+							QUICK_SCOPE_CYCLE_COUNTER(STAT_Oodle_InCaptureTime);
+
+							GPacketHandlerDiscardTimeguardMeasurement = true;
+#endif
+
 							InPacketLog->SerializePacket((void*)Packet.GetData(), DecompressedLength);
 						}
 #endif
@@ -1121,6 +1132,12 @@ void OodleNetworkHandlerComponent::Incoming(FIncomingPacketRef PacketRef)
 
 				if (SizeOfPacket > 0)
 				{
+#if !UE_BUILD_SHIPPING
+					QUICK_SCOPE_CYCLE_COUNTER(STAT_Oodle_InCaptureTime);
+
+					GPacketHandlerDiscardTimeguardMeasurement = true;
+#endif
+
 					InPacketLog->SerializePacket((void*)Packet.GetData(), SizeOfPacket);
 				}
 			}
@@ -1142,6 +1159,12 @@ void OodleNetworkHandlerComponent::Outgoing(FBitWriter& Packet, FOutPacketTraits
 
 			if (SizeOfPacket > 0)
 			{
+#if !UE_BUILD_SHIPPING
+				QUICK_SCOPE_CYCLE_COUNTER(STAT_Oodle_OutCaptureTime);
+
+				GPacketHandlerDiscardTimeguardMeasurement = true;
+#endif
+
 				OutPacketLog->SerializePacket((void*)Packet.GetData(), SizeOfPacket);
 			}
 		}
@@ -1158,7 +1181,6 @@ void OodleNetworkHandlerComponent::Outgoing(FBitWriter& Packet, FOutPacketTraits
 		FOodleNetworkAnalyticsVars* AnalyticsVars = (bOodleNetworkAnalytics && NetAnalyticsData.IsValid()) ? NetAnalyticsData->GetLocalData() : nullptr;
 		const bool bIsServer = (Handler->Mode == Handler::Mode::Server);
 		FOodleNetworkDictionary* CurDict = (bIsServer ? ServerDictionary.Get() : ClientDictionary.Get());
-		uint32 MinSizeForCompression = CVarOodleMinSizeForCompression.GetValueOnAnyThread();
 		uint32 UncompressedBytes = Packet.GetNumBytes();
 		bool bSkipCompressionClientDisabled = false;
 		bool bSkipCompressionTooSmall = false;
@@ -1177,7 +1199,7 @@ void OodleNetworkHandlerComponent::Outgoing(FBitWriter& Packet, FOutPacketTraits
 		}
 
 		// Skip compression when the packet is below the minimum size
-		if (!bSkipCompression && (MinSizeForCompression > 0 && UncompressedBytes < MinSizeForCompression))
+		if (!bSkipCompression && (GOodleMinSizeForCompression > 0 && static_cast<int32>(UncompressedBytes) < GOodleMinSizeForCompression))
 		{
 			bSkipCompression = true;
 			bSkipCompressionTooSmall = true;
@@ -1202,7 +1224,7 @@ void OodleNetworkHandlerComponent::Outgoing(FBitWriter& Packet, FOutPacketTraits
 				FMemory::Memcpy(UncompressedData, Packet.GetData(), UncompressedBytes);
 					
 				{
-#if STATS && !UE_BUILD_SHIPPING
+#if !UE_BUILD_SHIPPING
 					SCOPE_CYCLE_COUNTER(STAT_Oodle_OutCompressTime);
 #endif
 
diff --git a/Engine/Source/Runtime/PacketHandlers/PacketHandler/Private/PacketHandler.cpp b/Engine/Source/Runtime/PacketHandlers/PacketHandler/Private/PacketHandler.cpp
--- a/Engine/Source/Runtime/PacketHandlers/PacketHandler/Private/PacketHandler.cpp
+++ b/Engine/Source/Runtime/PacketHandlers/PacketHandler/Private/PacketHandler.cpp
@@ -40,6 +40,7 @@ FAutoConsoleVariableRef CVarNetPacketHandlerCRCDump(
 
 static int32 GPacketHandlerTimeguardLimit = 20;
 static float GPacketHandlerTimeguardThresholdMS = 0.0f;
+bool GPacketHandlerDiscardTimeguardMeasurement = false;
 
 static FAutoConsoleVariableRef CVarNetPacketHandlerTimeguardThresholdMS(
 	TEXT("net.PacketHandlerTimeguardThresholdMS"),
@@ -56,10 +57,11 @@ static FAutoConsoleVariableRef CVarNetPacketHandlerTimeguardLimit(
 // Lightweight time guard. Note: Threshold of 0 disables the timeguard
 #define NET_LIGHTWEIGHT_TIME_GUARD_BEGIN( Name, ThresholdMS ) \
 	float PREPROCESSOR_JOIN(__TimeGuard_ThresholdMS_, Name) = ThresholdMS; \
-	uint64 PREPROCESSOR_JOIN(__TimeGuard_StartCycles_, Name) = ( ThresholdMS > 0.0f && GPacketHandlerTimeguardLimit > 0 ) ? FPlatformTime::Cycles64() : 0;
+	uint64 PREPROCESSOR_JOIN(__TimeGuard_StartCycles_, Name) = ( ThresholdMS > 0.0f && GPacketHandlerTimeguardLimit > 0 ) ? FPlatformTime::Cycles64() : 0; \
+	GPacketHandlerDiscardTimeguardMeasurement = false;
 
 #define NET_LIGHTWEIGHT_TIME_GUARD_END( Name, NameStringCode ) \
-	if ( PREPROCESSOR_JOIN(__TimeGuard_ThresholdMS_, Name) > 0.0f && GPacketHandlerTimeguardLimit > 0 ) \
+	if ( PREPROCESSOR_JOIN(__TimeGuard_ThresholdMS_, Name) > 0.0f && GPacketHandlerTimeguardLimit > 0 && !GPacketHandlerDiscardTimeguardMeasurement ) \
 	{\
 		float PREPROCESSOR_JOIN(__TimeGuard_MSElapsed_,Name) = FPlatformTime::ToMilliseconds64( FPlatformTime::Cycles64() - PREPROCESSOR_JOIN(__TimeGuard_StartCycles_,Name) ); \
 		if ( PREPROCESSOR_JOIN(__TimeGuard_MSElapsed_,Name) > PREPROCESSOR_JOIN(__TimeGuard_ThresholdMS_, Name) ) \
@@ -85,7 +87,7 @@ PacketHandler::PacketHandler(FDDoSDetection* InDDoS/*=nullptr*/)
 	, DDoS(InDDoS)
 	, LowLevelSendDel()
 	, HandshakeCompleteDel()
-	, OutgoingPacket()
+	, OutgoingPacket(MAX_PACKET_SIZE * 8)
 	, IncomingPacket()
 	, HandlerComponents()
 	, MaxPacketBits(0)
diff --git a/Engine/Source/Runtime/PacketHandlers/PacketHandler/Public/PacketHandler.h b/Engine/Source/Runtime/PacketHandlers/PacketHandler/Public/PacketHandler.h
--- a/Engine/Source/Runtime/PacketHandlers/PacketHandler/Public/PacketHandler.h
+++ b/Engine/Source/Runtime/PacketHandlers/PacketHandler/Public/PacketHandler.h
@@ -30,6 +30,9 @@ namespace UE
 	}
 }
 
+#if !UE_BUILD_SHIPPING
+extern PACKETHANDLER_API bool GPacketHandlerDiscardTimeguardMeasurement;
+#endif
 
 /**
  * Delegates
