Removed OodleNetworkHandlerComponent's private pch file.

diff --git a/Engine/Plugins/Compression/OodleNetwork/Source/OodleNetworkHandlerComponent.Build.cs b/Engine/Plugins/Compression/OodleNetwork/Source/OodleNetworkHandlerComponent.Build.cs
--- a/Engine/Plugins/Compression/OodleNetwork/Source/OodleNetworkHandlerComponent.Build.cs
+++ b/Engine/Plugins/Compression/OodleNetwork/Source/OodleNetworkHandlerComponent.Build.cs
@@ -32,9 +32,6 @@ public class OodleNetworkHandlerComponent : ModuleRules
 		PublicDependencyModuleNames.Add("Engine");
 		PublicDependencyModuleNames.Add("Analytics");
 
-		PCHUsage = PCHUsageMode.UseSharedPCHs;
-		PrivatePCHHeaderFile = "Private/OodleNetworkHandlerComponentPCH.h";
-
 		// We depend on the .udic files, otherwise those dictionaries don't pass-through to final builds!
 		RuntimeDependencies.Add("$(ProjectDir)/Content/Oodle/...", StagedFileType.UFS);
 
diff --git a/Engine/Plugins/Compression/OodleNetwork/Source/Private/OodleNetworkHandlerComponentPCH.h b/Engine/Plugins/Compression/OodleNetwork/Source/Private/OodleNetworkHandlerComponentPCH.h
deleted file mode 100644
