Added stripped-down performance timing for Oodle 'Incoming' packets, to track a phantom hitch - and verify whether it is from Oodle or the stats system.

diff --git a/Engine/Plugins/Compression/OodleNetwork/Source/Private/OodleNetworkHandlerComponent.cpp b/Engine/Plugins/Compression/OodleNetwork/Source/Private/OodleNetworkHandlerComponent.cpp
--- a/Engine/Plugins/Compression/OodleNetwork/Source/Private/OodleNetworkHandlerComponent.cpp
+++ b/Engine/Plugins/Compression/OodleNetwork/Source/Private/OodleNetworkHandlerComponent.cpp
@@ -76,6 +76,59 @@ DEFINE_STAT(STAT_Oodle_StateBytes);
 #endif
 
 
+// OodleNetwork specific adaptation of timing macros from PacketHandler.cpp
+#if !UE_BUILD_SHIPPING
+namespace UE::Oodle
+{
+	static float GOodleTimeguardThresholdMS = 0.f;
+	static int32 GOodleTimeguardLimit = 20;
+
+	static FAutoConsoleVariableRef CVarOodleNetworkTimeguardThresholdMS(
+		TEXT("net.OodleNetwork.TimeGuardThresholdMS"),
+		GOodleTimeguardThresholdMS,
+		TEXT("Threshold in milliseconds for the OodleNetworkHandlerComponent timeguard."));
+
+	static FAutoConsoleVariableRef CVarOodleNetworkTimeguardLimit(
+		TEXT("net.OodleNetwork.TimeGuardLimit"),
+		GOodleTimeguardLimit,
+		TEXT("Sets the maximum number of OodleNetworkHandlerComponent timeguard logs."));
+}
+
+/** To be placed outside and before the scope of the measured code */
+#define UE_OODLE_LIGHTWEIGHT_TIME_GUARD_DECLARE(Name, ThresholdMS) \
+	const double PREPROCESSOR_JOIN(__TimeGuard_ThresholdMS_, Name) = ThresholdMS; \
+	double PREPROCESSOR_JOIN(__TimeGuard_MSElapsed_, Name) = 0.0;
+
+/** To be placed inside the scope, directly before the measured code */
+#define UE_OODLE_LIGHTWEIGHT_TIME_GUARD_BEGIN(Name) \
+	uint64 PREPROCESSOR_JOIN(__TimeGuard_StartCycles_, Name) = \
+		(PREPROCESSOR_JOIN(__TimeGuard_ThresholdMS_, Name) > 0.0 && UE::Oodle::GOodleTimeguardLimit > 0) ? FPlatformTime::Cycles64() : 0;
+
+/** To be placed inside the scope, directly after the measured code */
+#define UE_OODLE_LIGHTWEIGHT_TIME_GUARD_END(Name) \
+	if (PREPROCESSOR_JOIN(__TimeGuard_ThresholdMS_, Name) > 0.0 && UE::Oodle::GOodleTimeguardLimit > 0) \
+	{ \
+		PREPROCESSOR_JOIN(__TimeGuard_MSElapsed_, Name) = \
+			FPlatformTime::ToMilliseconds64(FPlatformTime::Cycles64() - PREPROCESSOR_JOIN(__TimeGuard_StartCycles_, Name)); \
+	}
+
+/** To be placed outside and after the scope of the measured code */
+#define UE_OODLE_LIGHTWEIGHT_TIME_GUARD_REPORT(Name) \
+	if (PREPROCESSOR_JOIN(__TimeGuard_MSElapsed_, Name) > PREPROCESSOR_JOIN(__TimeGuard_ThresholdMS_, Name)) \
+	{ \
+		UE_LOG(OodleNetworkHandlerComponentLog, Warning, TEXT("OodleNetworkHandlerComponent '%s' took %.2fms!"), TEXT(#Name), \
+				PREPROCESSOR_JOIN(__TimeGuard_MSElapsed_, Name)); \
+		UE::Oodle::GOodleTimeguardLimit--; \
+	}
+
+#else
+#define UE_OODLE_LIGHTWEIGHT_TIME_GUARD_DECLARE(Name, ThresholdMS)
+#define UE_OODLE_LIGHTWEIGHT_TIME_GUARD_BEGIN(Name)
+#define UE_OODLE_LIGHTWEIGHT_TIME_GUARD_END(Name)
+#define UE_OODLE_LIGHTWEIGHT_TIME_GUARD_REPORT(Name)
+#endif
+
+
 FString GOodleSaveDir = TEXT("");
 FString GOodleContentDir = TEXT("");
 
@@ -1024,15 +1077,25 @@ void OodleNetworkHandlerComponent::Incoming(FIncomingPacketRef PacketRef)
 
 					if (bSuccess)
 					{
+						UE_OODLE_LIGHTWEIGHT_TIME_GUARD_DECLARE(Incoming, UE::Oodle::GOodleTimeguardThresholdMS);
+
 						{
 #if !UE_BUILD_SHIPPING
 							SCOPE_CYCLE_COUNTER(STAT_Oodle_InDecompressTime);
 #endif
 
+							// The lightweight time guard will exclude STAT_Oodle_InDecompressTime processing time,
+							// allowing detection of stats system hitches, while verifying good performance of 'OodleNetwork1UDP_Decode' (hopefully)
+							UE_OODLE_LIGHTWEIGHT_TIME_GUARD_BEGIN(Incoming);
+
 							bSuccess = !!OodleNetwork1UDP_Decode(CurDict->CompressorState, CurDict->SharedDictionary, CompressedData,
 															CompressedLength, DecompressedData, DecompressedLength);
+
+							UE_OODLE_LIGHTWEIGHT_TIME_GUARD_END(Incoming);
 						}
 
+						UE_OODLE_LIGHTWEIGHT_TIME_GUARD_REPORT(Incoming);
+
 
 						if (!bSuccess)
 						{
