Fixed missing error check in Oodle dictionary archive serialization.

diff --git a/Engine/Plugins/Compression/OodleNetwork/Source/Private/OodleNetworkArchives.cpp b/Engine/Plugins/Compression/OodleNetwork/Source/Private/OodleNetworkArchives.cpp
--- a/Engine/Plugins/Compression/OodleNetwork/Source/Private/OodleNetworkArchives.cpp
+++ b/Engine/Plugins/Compression/OodleNetwork/Source/Private/OodleNetworkArchives.cpp
@@ -46,6 +46,11 @@ bool FOodleNetworkArchiveBase::SerializeOodleCompressData(FOodleCompressedData&
 		bSuccess = ! InnerArchive.IsError();
 	}
 
+	if (!bSuccess)
+	{
+		SetError();
+	}
+
 	return bSuccess;
 }
 
@@ -111,6 +116,11 @@ bool FOodleNetworkArchiveBase::SerializeOodleDecompressData(FOodleCompressedData
 	
 	bSuccess = bSuccess && !InnerArchive.IsError();
 
+	if (!bSuccess)
+	{
+		SetError();
+	}
+
 	return bSuccess;
 }
 
