No PCH build fix (Shipping)

#rnx

diff --git a/Engine/Plugins/Compression/OodleNetwork/Source/Private/OodleNetworkHandlerComponent.cpp b/Engine/Plugins/Compression/OodleNetwork/Source/Private/OodleNetworkHandlerComponent.cpp
--- a/Engine/Plugins/Compression/OodleNetwork/Source/Private/OodleNetworkHandlerComponent.cpp
+++ b/Engine/Plugins/Compression/OodleNetwork/Source/Private/OodleNetworkHandlerComponent.cpp
@@ -3,11 +3,13 @@
 #include "OodleNetworkHandlerComponent.h"
 #include "HAL/PlatformFile.h"
 #include "Modules/ModuleManager.h"
+#include "HAL/IConsoleManager.h"
 #include "HAL/PlatformFileManager.h"
 #include "HAL/FileManager.h"
 #include "Misc/ConfigCacheIni.h"
 #include "Misc/App.h"
 #include "Misc/EngineVersion.h"
+#include "UObject/CoreNet.h"
 
 #include "OodleNetworkTrainerCommandlet.h"
 #include "OodleNetworkAnalytics.h"
diff --git a/Engine/Source/Developer/TraceAnalysis/Public/Trace/Analyzer.h b/Engine/Source/Developer/TraceAnalysis/Public/Trace/Analyzer.h
--- a/Engine/Source/Developer/TraceAnalysis/Public/Trace/Analyzer.h
+++ b/Engine/Source/Developer/TraceAnalysis/Public/Trace/Analyzer.h
@@ -5,6 +5,7 @@
 #include "CoreTypes.h"
 #include "Containers/ArrayView.h"
 #include "Containers/StringView.h"
+#include "Trace/Trace.h"
 
 #include <type_traits>
 
diff --git a/Engine/Source/Runtime/Experimental/Dataflow/Engine/Private/Dataflow/DataflowEdNode.cpp b/Engine/Source/Runtime/Experimental/Dataflow/Engine/Private/Dataflow/DataflowEdNode.cpp
--- a/Engine/Source/Runtime/Experimental/Dataflow/Engine/Private/Dataflow/DataflowEdNode.cpp
+++ b/Engine/Source/Runtime/Experimental/Dataflow/Engine/Private/Dataflow/DataflowEdNode.cpp
@@ -11,7 +11,7 @@
 
 #include UE_INLINE_GENERATED_CPP_BY_NAME(DataflowEdNode)
 
-#if WITH_EDITOR && !(UE_BUILD_SHIPPING || UE_BUILD_TEST)
+#if WITH_EDITOR
 #include "EdGraph/EdGraphPin.h"
 #endif
 
diff --git a/Engine/Source/Runtime/RHI/Public/RHICommandList.h b/Engine/Source/Runtime/RHI/Public/RHICommandList.h
--- a/Engine/Source/Runtime/RHI/Public/RHICommandList.h
+++ b/Engine/Source/Runtime/RHI/Public/RHICommandList.h
@@ -29,6 +29,7 @@
 #include "ProfilingDebugging/CsvProfiler.h"
 #include "ProfilingDebugging/CpuProfilerTrace.h"
 #include "RHIBreadcrumbs.h"
+#include "RHIGlobals.h"
 #include "RHIShaderParameters.h"
 #include "RHITextureReference.h"
 #include "Trace/Trace.h"
