Enable use of iOS Simulator library builds of OodleNetwork and OddleDataCompression

#rnx

diff --git a/Engine/Plugins/Compression/OodleNetwork/Source/OodleNetworkHandlerComponent.Build.cs b/Engine/Plugins/Compression/OodleNetwork/Source/OodleNetworkHandlerComponent.Build.cs
--- a/Engine/Plugins/Compression/OodleNetwork/Source/OodleNetworkHandlerComponent.Build.cs
+++ b/Engine/Plugins/Compression/OodleNetwork/Source/OodleNetworkHandlerComponent.Build.cs
@@ -78,9 +78,10 @@ public class OodleNetworkHandlerComponent : ModuleRules
 		}
 		else if (Target.Platform == UnrealTargetPlatform.IOS)
 		{
-			// NOTE: Simulator builds should use .sim.a files but we're not building sim as of this writing
-			ReleaseLib = "liboo2netios.a";
-			DebugLib = "liboo2netios_dbg.a";
+			string LibExt = (Target.Architecture == UnrealArch.IOSSimulator) ? ".sim.a" : ".a";
+
+			ReleaseLib = "liboo2netios" + LibExt;
+			DebugLib = "liboo2netios_dbg" + LibExt;
 		}
 		else if (Target.Platform == UnrealTargetPlatform.TVOS)
 		{
diff --git a/Engine/Source/Runtime/OodleDataCompression/OodleDataCompression.Build.cs b/Engine/Source/Runtime/OodleDataCompression/OodleDataCompression.Build.cs
--- a/Engine/Source/Runtime/OodleDataCompression/OodleDataCompression.Build.cs
+++ b/Engine/Source/Runtime/OodleDataCompression/OodleDataCompression.Build.cs
@@ -74,9 +74,9 @@ using UnrealBuildTool;
 			}
 			else if (Target.Platform == UnrealTargetPlatform.IOS)
 			{
-				// NOTE: Simulator builds should use .sim.a files but we're not building sim as of this writing
-				ReleaseLib = "liboo2coreios.a";
-				DebugLib = "liboo2coreios_dbg.a";
+				string libEnding = (Target.Architecture == UnrealArch.IOSSimulator) ? ".sim.a" : ".a";
+				ReleaseLib = "liboo2coreios" + libEnding;
+				DebugLib = "liboo2coreios_dbg" + libEnding;
 			}
 			else if (Target.Platform == UnrealTargetPlatform.TVOS)
 			{
