- Oodle data/net libs for VisionOS

diff --git a/Engine/Platforms/VisionOS/Plugins/Compression/OodleNetwork/OodleNetwork_VisionOS.uplugin b/Engine/Platforms/VisionOS/Plugins/Compression/OodleNetwork/OodleNetwork_VisionOS.uplugin
new file mode 100644
--- /dev/null
+++ b/Engine/Platforms/VisionOS/Plugins/Compression/OodleNetwork/OodleNetwork_VisionOS.uplugin
@@ -0,0 +1,4 @@
+{
+	"FileVersion": 3,
+	"bIsPluginExtension": true
+}
diff --git a/Engine/Platforms/VisionOS/Plugins/Compression/OodleNetwork/Source/OodleNetworkHandlerComponent_VisionOS.Build.cs b/Engine/Platforms/VisionOS/Plugins/Compression/OodleNetwork/Source/OodleNetworkHandlerComponent_VisionOS.Build.cs
new file mode 100644
--- /dev/null
+++ b/Engine/Platforms/VisionOS/Plugins/Compression/OodleNetwork/Source/OodleNetworkHandlerComponent_VisionOS.Build.cs
@@ -0,0 +1,15 @@
+// Copyright Epic Games, Inc. All Rights Reserved.
+
+using System.IO;
+using UnrealBuildTool;
+
+public class OodleNetworkHandlerComponent_VisionOS : OodleNetworkHandlerComponent
+{
+	string LibExt { get { return Target.Architecture == UnrealArch.IOSSimulator ? ".sim.a" : ".a"; } }
+	protected override string ReleaseLibraryName { get { return "liboo2netvisionos" + LibExt; } }
+	protected override string DebugLibraryName { get { return "liboo2netvisionos_dbg" + LibExt; } }
+
+	public OodleNetworkHandlerComponent_VisionOS(ReadOnlyTargetRules Target) : base(Target)
+	{
+	}
+}
