PR #11492: Add UObject/CoreNet.h include to fix -Nopch build error

diff --git a/Engine/Plugins/Compression/OodleNetwork/Source/Public/OodleNetworkHandlerComponent.h b/Engine/Plugins/Compression/OodleNetwork/Source/Public/OodleNetworkHandlerComponent.h
--- a/Engine/Plugins/Compression/OodleNetwork/Source/Public/OodleNetworkHandlerComponent.h
+++ b/Engine/Plugins/Compression/OodleNetwork/Source/Public/OodleNetworkHandlerComponent.h
@@ -5,6 +5,7 @@
 // @todo #JohnB: Separate module-based header code, from other class implementations, so that you can setup the PCH.h file correctly
 
 #include "Net/Core/Analytics/NetAnalytics.h"
+#include "UObject/CoreNet.h"
 #include "PacketHandler.h"
 #include "OodleNetworkArchives.h"
 #include "OodleNetworkFaultHandler.h"
