OodleAnalytics
* Added [In/Out]PostLengthTotal analytics to more easily graph total traffic generated

diff --git a/Engine/Plugins/Compression/OodleNetwork/Source/Private/OodleNetworkAnalytics.cpp b/Engine/Plugins/Compression/OodleNetwork/Source/Private/OodleNetworkAnalytics.cpp
--- a/Engine/Plugins/Compression/OodleNetwork/Source/Private/OodleNetworkAnalytics.cpp
+++ b/Engine/Plugins/Compression/OodleNetwork/Source/Private/OodleNetworkAnalytics.cpp
@@ -108,12 +108,16 @@ void FOodleNetAnalyticsData::SendAnalytics()
 
 		uint64 InPreLengthTotal = InCompressedLengthTotal + InNotCompressedLengthTotal;
 		uint64 InPreWithOverheadLengthTotal = InCompressedWithOverheadLengthTotal + InNotCompressedLengthTotal;
+
+		/** The total size in Bytes of In traffic received */
 		uint64 InPostLengthTotal = InDecompressedLengthTotal + InNotCompressedLengthTotal;
 
 		uint64 OutPreLengthTotal = OutBeforeCompressedLengthTotal + OutNotCompressedFailedLengthTotal + OutNotCompressedSkippedLengthTotal;
-		uint64 OutPostLengthTotal = OutCompressedLengthTotal + OutNotCompressedFailedLengthTotal + OutNotCompressedSkippedLengthTotal;
 		uint64 OutPostWithOverheadLengthTotal = OutCompressedWithOverheadLengthTotal + OutNotCompressedFailedLengthTotal + OutNotCompressedSkippedLengthTotal;
 
+		/** The total size in Bytes of Out traffic generated */
+		uint64 OutPostLengthTotal = OutCompressedLengthTotal + OutNotCompressedFailedLengthTotal + OutNotCompressedSkippedLengthTotal;
+
 		uint64 OutPreAttemptedLengthTotal = OutBeforeCompressedLengthTotal + OutNotCompressedFailedLengthTotal;
 		uint64 OutPostAttemptedWithOverheadLengthTotal = OutCompressedWithOverheadLengthTotal + OutNotCompressedFailedLengthTotal;
 
@@ -187,6 +191,8 @@ void FOodleNetAnalyticsData::SendAnalytics()
 		UE_LOG(OodleNetworkHandlerComponentLog, Log, TEXT(" - OutNotCompressedFailedLengthTotal: %llu"), OutNotCompressedFailedLengthTotal);
 		UE_LOG(OodleNetworkHandlerComponentLog, Log, TEXT(" - OutNotCompressedSkippedLengthTotal: %llu"), OutNotCompressedSkippedLengthTotal);
 		UE_LOG(OodleNetworkHandlerComponentLog, Log, TEXT(" - OutNotCompressedNumTotal: %llu"), OutNotCompressedNumTotal);
+		UE_LOG(OodleNetworkHandlerComponentLog, Log, TEXT(" - OutPostLengthTotal: %llu"), OutPostLengthTotal);
+		UE_LOG(OodleNetworkHandlerComponentLog, Log, TEXT(" - InPostLengthTotal: %llu"), InPostLengthTotal);
 		UE_LOG(OodleNetworkHandlerComponentLog, Log, TEXT(" - InSavingsPercentTotal: %i"), InSavingsPercentTotal);
 		UE_LOG(OodleNetworkHandlerComponentLog, Log, TEXT(" - OutSavingsPercentTotal: %i"), OutSavingsPercentTotal);
 		UE_LOG(OodleNetworkHandlerComponentLog, Log, TEXT(" - InSavingsBytesTotal: %lli"), InSavingsBytesTotal);
@@ -223,6 +229,8 @@ void FOodleNetAnalyticsData::SendAnalytics()
 		static const FString EZAttrib_OutNotCompressedFailedLengthTotal = TEXT("OutNotCompressedFailedLengthTotal");
 		static const FString EZAttrib_OutNotCompressedSkippedLengthTotal = TEXT("OutNotCompressedSkippedLengthTotal");
 		static const FString EZAttrib_OutNotCompressedNumTotal = TEXT("OutNotCompressedNumTotal");
+		static const FString EZAttrib_OutPostLengthTotal = TEXT("OutPostLengthTotal");
+		static const FString EZAttrib_InPostLengthTotal = TEXT("InPostLengthTotal");
 		static const FString EZAttrib_InSavingsPercentTotal = TEXT("InSavingsPercentTotal");
 		static const FString EZAttrib_OutSavingsPercentTotal = TEXT("OutSavingsPercentTotal");
 		static const FString EZAttrib_InSavingsBytesTotal = TEXT("InSavingsBytesTotal");
@@ -259,6 +267,8 @@ void FOodleNetAnalyticsData::SendAnalytics()
 			EZAttrib_OutNotCompressedFailedLengthTotal, OutNotCompressedFailedLengthTotal,
 			EZAttrib_OutNotCompressedSkippedLengthTotal, OutNotCompressedSkippedLengthTotal,
 			EZAttrib_OutNotCompressedNumTotal, OutNotCompressedNumTotal,
+			EZAttrib_OutPostLengthTotal, OutPostLengthTotal,
+			EZAttrib_InPostLengthTotal, InPostLengthTotal,
 			EZAttrib_InSavingsPercentTotal, InSavingsPercentTotal,
 			EZAttrib_OutSavingsPercentTotal, OutSavingsPercentTotal,
 			EZAttrib_InSavingsBytesTotal, InSavingsBytesTotal,
