Added Oodle CSV Stats tracking
#tests Tested in a local test build, captured CSV profiling data and generated a PerfReport which shows the new stats.

diff --git a/Engine/Plugins/Compression/OodleNetwork/Source/Private/OodleNetworkHandlerComponent.cpp b/Engine/Plugins/Compression/OodleNetwork/Source/Private/OodleNetworkHandlerComponent.cpp
--- a/Engine/Plugins/Compression/OodleNetwork/Source/Private/OodleNetworkHandlerComponent.cpp
+++ b/Engine/Plugins/Compression/OodleNetwork/Source/Private/OodleNetworkHandlerComponent.cpp
@@ -18,6 +18,7 @@
 
 #if !UE_BUILD_SHIPPING
 #include "Engine/Engine.h"
+#include "ProfilingDebugging/CsvProfiler.h"
 #endif
 
 #include "Net/Core/Connection/NetCloseResult.h"
@@ -525,8 +526,9 @@ void OodleNetworkHandlerComponent::Initialize()
 			FParse::Value(FCommandLine::Get(), TEXT("CapturePercentage="), CapturePercentage);
 
 			int32 RandNum = FMath::RandRange(0, 100);
-			UE_LOG(OodleNetworkHandlerComponentLog, Log, TEXT("Enabling Oodle capture mode. Random number is: %d, Capture Percentage is: %d, random number must be less than capture percentage to capture."), RandNum, CapturePercentage);
-			if (RandNum <= CapturePercentage)
+			const bool bShouldEnableCapture = (RandNum <= CapturePercentage);
+			UE_LOG(OodleNetworkHandlerComponentLog, Log, TEXT("Enabling Oodle capture mode: [%s]. Random number is: %d, Capture Percentage is: %d, random number must be less than capture percentage to capture."), bShouldEnableCapture ? TEXT("TRUE") : TEXT("FALSE"), RandNum, CapturePercentage);
+			if (bShouldEnableCapture)
 			{
 				InitializePacketLogs();
 			}
@@ -1081,6 +1083,7 @@ void OodleNetworkHandlerComponent::Incoming(FIncomingPacketRef PacketRef)
 						{
 #if !UE_BUILD_SHIPPING
 							SCOPE_CYCLE_COUNTER(STAT_Oodle_InDecompressTime);
+							CSV_SCOPED_TIMING_STAT_EXCLUSIVE(Oodle_InDecompressTime);
 #endif
 
 							// The lightweight time guard will exclude STAT_Oodle_InDecompressTime processing time,
@@ -1129,7 +1132,7 @@ void OodleNetworkHandlerComponent::Incoming(FIncomingPacketRef PacketRef)
 						{
 #if !UE_BUILD_SHIPPING
 							QUICK_SCOPE_CYCLE_COUNTER(STAT_Oodle_InCaptureTime);
-
+							CSV_SCOPED_TIMING_STAT_EXCLUSIVE(Oodle_InCaptureTime);
 							GPacketHandlerDiscardTimeguardMeasurement = true;
 #endif
 
@@ -1196,7 +1199,7 @@ void OodleNetworkHandlerComponent::Incoming(FIncomingPacketRef PacketRef)
 				{
 #if !UE_BUILD_SHIPPING
 					QUICK_SCOPE_CYCLE_COUNTER(STAT_Oodle_InCaptureTime);
-
+					CSV_SCOPED_TIMING_STAT_EXCLUSIVE(Oodle_InCaptureTime);
 					GPacketHandlerDiscardTimeguardMeasurement = true;
 #endif
 
@@ -1223,7 +1226,7 @@ void OodleNetworkHandlerComponent::Outgoing(FBitWriter& Packet, FOutPacketTraits
 			{
 #if !UE_BUILD_SHIPPING
 				QUICK_SCOPE_CYCLE_COUNTER(STAT_Oodle_OutCaptureTime);
-
+				CSV_SCOPED_TIMING_STAT_EXCLUSIVE(Oodle_OutCaptureTime);
 				GPacketHandlerDiscardTimeguardMeasurement = true;
 #endif
 
@@ -1288,6 +1291,7 @@ void OodleNetworkHandlerComponent::Outgoing(FBitWriter& Packet, FOutPacketTraits
 				{
 #if !UE_BUILD_SHIPPING
 					SCOPE_CYCLE_COUNTER(STAT_Oodle_OutCompressTime);
+					CSV_SCOPED_TIMING_STAT_EXCLUSIVE(Oodle_OutCompressTime);
 #endif
 
 					CompressedLengthSINT = OodleNetwork1UDP_Encode(CurDict->CompressorState, CurDict->SharedDictionary,
