OodleNetwork: Fixed unsafe cast warnings.

#rnx

diff --git a/Engine/Plugins/Compression/OodleNetwork/Source/Private/OodleNetworkAnalytics.cpp b/Engine/Plugins/Compression/OodleNetwork/Source/Private/OodleNetworkAnalytics.cpp
--- a/Engine/Plugins/Compression/OodleNetwork/Source/Private/OodleNetworkAnalytics.cpp
+++ b/Engine/Plugins/Compression/OodleNetwork/Source/Private/OodleNetworkAnalytics.cpp
@@ -128,10 +128,10 @@ void FOodleNetAnalyticsData::SendAnalytics()
 		 * Also factors in skipped/failed compression (which has increased a lot since analytics was added) - reducing the usefulness for determining algorithm compression.
 		 */
 			/** The percentage of compression savings, of all incoming packets. */
-			int8 InSavingsPercentTotal = (1.0 - ((double)InPreLengthTotal / (double)InPostLengthTotal)) * 100.0;
+			int8 InSavingsPercentTotal = (int8)((1.0 - ((double)InPreLengthTotal / (double)InPostLengthTotal)) * 100.0);
 
 			/** The percentage of compression savings, of all outgoing packets. */
-			int8 OutSavingsPercentTotal = (1.0 - ((double)OutPostLengthTotal / (double)OutPreLengthTotal)) * 100.0;
+			int8 OutSavingsPercentTotal = (int8)((1.0 - ((double)OutPostLengthTotal / (double)OutPreLengthTotal)) * 100.0);
 
 			/** The number of bytes saved due to compression, of all incoming packets. */
 			int64 InSavingsBytesTotal = InPostLengthTotal - InPreLengthTotal;
@@ -144,10 +144,10 @@ void FOodleNetAnalyticsData::SendAnalytics()
 		 * This is the most accurate measure of compression savings (in terms of overall bandwidth).
 		 */
 			/** The percentage of compression savings, of all incoming packets. */
-			int8 InSavingsWithOverheadPercentTotal = (1.0 - ((double)InPreWithOverheadLengthTotal / (double)InPostLengthTotal)) * 100.0;
+			int8 InSavingsWithOverheadPercentTotal = (int8)((1.0 - ((double)InPreWithOverheadLengthTotal / (double)InPostLengthTotal)) * 100.0);
 
 			/** The percentage of compression savings, of all outgoing packets. */
-			int8 OutSavingsWithOverheadPercentTotal = (1.0 - ((double)OutPostWithOverheadLengthTotal / (double)OutPreLengthTotal)) * 100.0;
+			int8 OutSavingsWithOverheadPercentTotal = (int8)((1.0 - ((double)OutPostWithOverheadLengthTotal / (double)OutPreLengthTotal)) * 100.0);
 
 			/** The number of bytes saved due to compression, of all incoming packets. */
 			int64 InSavingsWithOverheadBytesTotal = InPostLengthTotal - InPreWithOverheadLengthTotal;
@@ -160,10 +160,10 @@ void FOodleNetAnalyticsData::SendAnalytics()
 		 * This is the best measure for determining algorithm compression performance, especially when measured against the Oodle encode/decode CPU cost.
 		 */
 			/** The percentage of compression savings, of all incoming packets that were compressed (attempted but failed compress, can't be counted here) */
-			int8 InAttemptedSavingsWithOverheadPercentTotal = (1.0 - ((double)InCompressedWithOverheadLengthTotal / (double)InDecompressedLengthTotal)) * 100.0;
+			int8 InAttemptedSavingsWithOverheadPercentTotal = (int8)((1.0 - ((double)InCompressedWithOverheadLengthTotal / (double)InDecompressedLengthTotal)) * 100.0);
 
 			/** The percentage of compression savings, of all outgoing packets that attempted compression. */
-			int8 OutAttemptedSavingsWithOverheadPercentTotal = (1.0 - ((double)OutPostAttemptedWithOverheadLengthTotal / (double)OutPreAttemptedLengthTotal)) * 100.0;
+			int8 OutAttemptedSavingsWithOverheadPercentTotal = (int8)((1.0 - ((double)OutPostAttemptedWithOverheadLengthTotal / (double)OutPreAttemptedLengthTotal)) * 100.0);
 
 
 		uint32 NumOodleNetworkHandlersCompressionDisabled = NumOodleNetworkHandlers - NumOodleNetworkHandlersCompressionEnabled;
diff --git a/Engine/Plugins/Compression/OodleNetwork/Source/Private/OodleNetworkArchives.cpp b/Engine/Plugins/Compression/OodleNetwork/Source/Private/OodleNetworkArchives.cpp
--- a/Engine/Plugins/Compression/OodleNetwork/Source/Private/OodleNetworkArchives.cpp
+++ b/Engine/Plugins/Compression/OodleNetwork/Source/Private/OodleNetworkArchives.cpp
@@ -38,7 +38,7 @@ bool FOodleNetworkArchiveBase::SerializeOodleCompressData(FOodleCompressedData&
 		// no compression
 		OutDataInfo.CompressedLength.Set(*this, DataBytes);
 
-		uint32 OffsetPos = InnerArchive.Tell();
+		uint32 OffsetPos = static_cast<uint32>(InnerArchive.Tell());
 		OutDataInfo.Offset.Set(*this, OffsetPos);
 
 		InnerArchive.Serialize((void*)Data, DataBytes);
@@ -155,7 +155,7 @@ void FPacketCaptureArchive::FCaptureHeader::SerializeHeader(FPacketCaptureArchiv
 		Ar << PacketDataLength;
 	}
 
-	PacketDataOffset.Set(Ar, Ar.Tell());
+	PacketDataOffset.Set(Ar, static_cast<uint32>(Ar.Tell()));
 }
 
 
@@ -216,7 +216,7 @@ void FPacketCaptureArchive::SerializePacket(void* PacketData, uint32& PacketSize
 	if (IsSaving())
 	{
 		uint32 NewPacketCount = Header.PacketCount.Get() + 1;
-		uint32 NewPacketDataLength = Header.PacketDataLength.Get() + (Tell() - StartPos);
+		uint32 NewPacketDataLength = Header.PacketDataLength.Get() + static_cast<uint32>(Tell() - StartPos);
 
 		Header.PacketCount.Set(*this, NewPacketCount);
 		Header.PacketDataLength.Set(*this, NewPacketDataLength);
diff --git a/Engine/Plugins/Compression/OodleNetwork/Source/Private/OodleNetworkHandlerComponent.cpp b/Engine/Plugins/Compression/OodleNetwork/Source/Private/OodleNetworkHandlerComponent.cpp
--- a/Engine/Plugins/Compression/OodleNetwork/Source/Private/OodleNetworkHandlerComponent.cpp
+++ b/Engine/Plugins/Compression/OodleNetwork/Source/Private/OodleNetworkHandlerComponent.cpp
@@ -242,8 +242,8 @@ static FOodleNetStats GOodleNetStats;
 void FOodleNetStats::UpdateStats(float DeltaTime)
 {
 	// Input
-	const uint32 InRaw = FMath::TruncToInt(InDecompressedLength / DeltaTime);
-	const uint32 InCompressed = FMath::TruncToInt(InCompressedLength / DeltaTime);
+	const uint32 InRaw = FMath::TruncToInt((float)InDecompressedLength / DeltaTime);
+	const uint32 InCompressed = FMath::TruncToInt((float)InCompressedLength / DeltaTime);
 
 	SET_DWORD_STAT(STAT_Oodle_InRaw, InRaw);
 	SET_DWORD_STAT(STAT_Oodle_InCompressed, InCompressed);
@@ -254,8 +254,8 @@ void FOodleNetStats::UpdateStats(float DeltaTime)
 
 
 	// Output
-	uint32 OutRaw = FMath::TruncToInt(OutUncompressedLength / DeltaTime);
-	uint32 OutCompressed = FMath::TruncToInt(OutCompressedLength / DeltaTime);
+	uint32 OutRaw = FMath::TruncToInt((float)OutUncompressedLength / DeltaTime);
+	uint32 OutCompressed = FMath::TruncToInt((float)OutCompressedLength / DeltaTime);
 
 	SET_DWORD_STAT(STAT_Oodle_OutRaw, OutRaw);
 	SET_DWORD_STAT(STAT_Oodle_OutCompressed, OutCompressed);
@@ -564,7 +564,7 @@ void OodleNetworkHandlerComponent::Initialize()
 
 		if (!MeasureAr.IsError())
 		{
-			OodleReservedPacketBits += MeasureAr.GetNumBits();
+			OodleReservedPacketBits += static_cast<uint32>(MeasureAr.GetNumBits());
 
 #if !UE_BUILD_SHIPPING
 			SET_DWORD_STAT(STAT_PacketReservedOodle, OodleReservedPacketBits);
@@ -667,7 +667,7 @@ void OodleNetworkHandlerComponent::InitializeDictionary(FString FilePath, TShare
 				UE_LOG(OodleNetworkHandlerComponentLog, Log, TEXT("Loading dictionary file: %s"), *FilePath);
 
 				// Uncompact the compressor state
-				uint32 CompressorStateSize = OodleNetwork1UDP_State_Size();
+				uint32 CompressorStateSize = static_cast<uint32>(OodleNetwork1UDP_State_Size());
 				OodleNetwork1UDP_State* CompressorState = (OodleNetwork1UDP_State*)FMemory::Malloc(CompressorStateSize);
 
 				OO_BOOL UncompactOk = OodleNetwork1UDP_State_Uncompact_ForVersion(CompressorState, (OodleNetwork1UDP_StateCompacted*)CompactCompressorState, BoundArc.Header.OodleMajorHeaderVersion);
@@ -682,7 +682,7 @@ void OodleNetworkHandlerComponent::InitializeDictionary(FString FilePath, TShare
 
 				// Create the shared dictionary state
 				int32 HashTableSize = BoundArc.Header.HashTableSize.Get();
-				uint32 SharedDictionarySize = OodleNetwork1_Shared_Size(HashTableSize);
+				uint32 SharedDictionarySize = static_cast<uint32>(OodleNetwork1_Shared_Size(HashTableSize));
 				OodleNetwork1_Shared* SharedDictionary = (OodleNetwork1_Shared*)FMemory::Malloc(SharedDictionarySize);
 
 				OodleNetwork1_Shared_SetWindow(SharedDictionary, HashTableSize, (void*)DictionaryData, DictionaryBytes);
@@ -1043,7 +1043,7 @@ void OodleNetworkHandlerComponent::Incoming(FIncomingPacketRef PacketRef)
 
 			if (CurDict != nullptr)
 			{
-				uint32 BeforeDecompressedLength = Packet.GetBytesLeft();
+				uint32 BeforeDecompressedLength = static_cast<uint32>(Packet.GetBytesLeft());
 				uint32 DecompressedLength;
 
 				SerializeOodlePacketSize(Packet, DecompressedLength);
@@ -1069,7 +1069,7 @@ void OodleNetworkHandlerComponent::Incoming(FIncomingPacketRef PacketRef)
 					static uint8 CompressedData[MAX_OODLE_BUFFER];
 					static uint8 DecompressedData[MAX_OODLE_BUFFER];
 
-					const int32 CompressedLength = Packet.GetBytesLeft();
+					const uint32 CompressedLength = static_cast<uint32>(Packet.GetBytesLeft());
 
 					// @todo Oodle FIX ME : byte buffer serialized out of bit buffer
 					Packet.Serialize(CompressedData, CompressedLength);
@@ -1184,7 +1184,7 @@ void OodleNetworkHandlerComponent::Incoming(FIncomingPacketRef PacketRef)
 			if (bOodleNetworkAnalytics && NetAnalyticsData.IsValid())
 			{
 				FOodleNetworkAnalyticsVars* AnalyticsVars = NetAnalyticsData->GetLocalData();
-				uint32 PacketSize = Packet.GetBytesLeft();
+				uint32 PacketSize = static_cast<uint32>(Packet.GetBytesLeft());
 
 				AnalyticsVars->InNotCompressedNum++;
 				AnalyticsVars->InNotCompressedLengthTotal += PacketSize;
@@ -1193,7 +1193,7 @@ void OodleNetworkHandlerComponent::Incoming(FIncomingPacketRef PacketRef)
 #if !UE_BUILD_SHIPPING || OODLE_DEV_SHIPPING
 			if (bCaptureMode && Handler->Mode == UE::Handler::Mode::Server && InPacketLog != nullptr)
 			{
-				uint32 SizeOfPacket = Packet.GetBytesLeft();
+				uint32 SizeOfPacket = static_cast<uint32>(Packet.GetBytesLeft());
 
 				if (SizeOfPacket > 0)
 				{
@@ -1220,7 +1220,7 @@ void OodleNetworkHandlerComponent::Outgoing(FBitWriter& Packet, FOutPacketTraits
 #if !UE_BUILD_SHIPPING || OODLE_DEV_SHIPPING
 		if (bCaptureMode && Handler->Mode == UE::Handler::Mode::Server && OutPacketLog != nullptr)
 		{
-			uint32 SizeOfPacket = Packet.GetNumBytes();
+			uint32 SizeOfPacket = static_cast<uint32>(Packet.GetNumBytes());
 
 			if (SizeOfPacket > 0)
 			{
@@ -1246,7 +1246,7 @@ void OodleNetworkHandlerComponent::Outgoing(FBitWriter& Packet, FOutPacketTraits
 		FOodleNetworkAnalyticsVars* AnalyticsVars = (bOodleNetworkAnalytics && NetAnalyticsData.IsValid()) ? NetAnalyticsData->GetLocalData() : nullptr;
 		const bool bIsServer = (Handler->Mode == UE::Handler::Mode::Server);
 		FOodleNetworkDictionary* CurDict = (bIsServer ? ServerDictionary.Get() : ClientDictionary.Get());
-		uint32 UncompressedBytes = Packet.GetNumBytes();
+		uint32 UncompressedBytes = static_cast<uint32>(Packet.GetNumBytes());
 		bool bSkipCompressionClientDisabled = false;
 		bool bSkipCompressionTooSmall = false;
 		bool bSkipCompression = false;
@@ -1276,7 +1276,7 @@ void OodleNetworkHandlerComponent::Outgoing(FBitWriter& Packet, FOutPacketTraits
 			check(MaxOutgoingBits <= (MAX_OODLE_PACKET_BYTES * 8));
 #endif
 			uint32 MaxAdjustedLengthBits = MaxOutgoingBits - OodleReservedPacketBits;
-			uint32 UncompressedBits = Packet.GetNumBits();
+			uint32 UncompressedBits = static_cast<uint32>(Packet.GetNumBits());
 
 			bool bWithinBitBounds = UncompressedBits > 0 && ensure(UncompressedBits <= MaxAdjustedLengthBits) &&
 				ensure(OodleNetwork1_CompressedBufferSizeNeeded(UncompressedBytes) <= MAX_OODLE_BUFFER);
@@ -1423,7 +1423,7 @@ void OodleNetworkHandlerComponent::Outgoing(FBitWriter& Packet, FOutPacketTraits
 			}
 
 			uint8 bCompresedPacket = false;
-			uint32 UncompressedBits = Packet.GetNumBits();
+			uint32 UncompressedBits = static_cast<uint32>(Packet.GetNumBits());
 
 			FMemory::Memcpy(UncompressedData, Packet.GetData(), Packet.GetNumBytes());
 
diff --git a/Engine/Plugins/Compression/OodleNetwork/Source/Private/OodleNetworkTrainerCommandlet.cpp b/Engine/Plugins/Compression/OodleNetwork/Source/Private/OodleNetworkTrainerCommandlet.cpp
--- a/Engine/Plugins/Compression/OodleNetwork/Source/Private/OodleNetworkTrainerCommandlet.cpp
+++ b/Engine/Plugins/Compression/OodleNetwork/Source/Private/OodleNetworkTrainerCommandlet.cpp
@@ -817,7 +817,7 @@ bool FOodleNetworkDictionaryGenerator::ReadPackets(const TArray<FString>& InputC
 	}
 
 	// Now translate the randomized index, back into a sequential index, which maps random packets to each list-type (from type 0-4)
-	int32 ListSplitCount = 3 + (bCompressionTest ? 1 : 0);
+	uint8 ListSplitCount = 3 + (bCompressionTest ? 1 : 0);
 	TArray<uint8> PacketToListMap;
 
 	PacketToListMap.SetNum(PacketCount);
@@ -837,7 +837,7 @@ bool FOodleNetworkDictionaryGenerator::ReadPackets(const TArray<FString>& InputC
 	{
 		for (int32 i=0; i<RandIndex.Num(); i++)
 		{
-			PacketToListMap[RandIndex[i]] = (i % ListSplitCount);
+			PacketToListMap[RandIndex[i]] = static_cast<uint8>(i % ListSplitCount);
 		}
 
 		// Now step through the packets, and randomly assign them to the appropriate list
@@ -1044,7 +1044,7 @@ bool FOodleNetworkDictionaryGenerator::GenerateAndWriteDictionary()
 			bSuccess = bSuccess && OutputFile.SerializeOodleCompressData(OutputFile.Header.DictionaryData, NewDictionaryData, DictionarySize);
 
 			bSuccess = bSuccess && OutputFile.SerializeOodleCompressData(OutputFile.Header.CompressorData, (uint8*)CompactCompressorState,
-																			CompactCompressorStateBytes);
+																			static_cast<uint32>(CompactCompressorStateBytes));
 
 
 			// Important warning for unusually small dictionary files - if they compress down this much, something is usually wrong.
diff --git a/Engine/Plugins/Compression/OodleNetwork/Source/Public/OodleNetworkArchives.h b/Engine/Plugins/Compression/OodleNetwork/Source/Public/OodleNetworkArchives.h
--- a/Engine/Plugins/Compression/OodleNetwork/Source/Public/OodleNetworkArchives.h
+++ b/Engine/Plugins/Compression/OodleNetwork/Source/Public/OodleNetworkArchives.h
@@ -108,7 +108,7 @@ protected:
 		{
 			if (Offset == 0)
 			{
-				Offset = Ar.Tell();
+				Offset = static_cast<uint32>(Ar.Tell());
 			}
 			else
 			{
diff --git a/Engine/Plugins/Compression/OodleNetwork/Source/Public/OodleNetworkHandlerComponent.h b/Engine/Plugins/Compression/OodleNetwork/Source/Public/OodleNetworkHandlerComponent.h
--- a/Engine/Plugins/Compression/OodleNetwork/Source/Public/OodleNetworkHandlerComponent.h
+++ b/Engine/Plugins/Compression/OodleNetwork/Source/Public/OodleNetworkHandlerComponent.h
@@ -185,7 +185,7 @@ public:
 	 */
 	FORCEINLINE void CheckForUpdate()
 	{
-		float DeltaTime = FPlatformTime::Seconds() - LastStatsUpdate;
+		float DeltaTime = static_cast<float>(FPlatformTime::Seconds() - LastStatsUpdate);
 
 		if (DeltaTime > 1.f)
 		{
