OodleNetworkTrainerCommandlet: use "official" packet count, don't keep going until file end

When the packet count header is present, that number is contractual. We can have
extra packets after that in the file, but they're supposed to be ignored and not
accounted for in the packet total we computed earlier.

diff --git a/Engine/Plugins/Compression/OodleNetwork/Source/Private/OodleNetworkTrainerCommandlet.cpp b/Engine/Plugins/Compression/OodleNetwork/Source/Private/OodleNetworkTrainerCommandlet.cpp
--- a/Engine/Plugins/Compression/OodleNetwork/Source/Private/OodleNetworkTrainerCommandlet.cpp
+++ b/Engine/Plugins/Compression/OodleNetwork/Source/Private/OodleNetworkTrainerCommandlet.cpp
@@ -867,8 +867,18 @@ bool FOodleNetworkDictionaryGenerator::ReadPackets(const TArray<FString>& InputC
 
 	for (TUniquePtr<FPacketCaptureArchive>& CurArc : BoundArchives)
 	{
-		while (CurArc->Tell() < CurArc->TotalSize() && PacketIdx < (uint32)PacketCount)
+		const uint32 PacketCountInCurArc = CurArc->GetPacketCount();
+
+		// We go by the packet count in the header if present.
+		// In some cases captures can contain extra packets past that point,
+		// but stay with the "official" (properly flushed) count in the header.
+		// Files that don't have the header did an earlier pass over the file
+		// contents to count complete packets.
+		for (uint32 PacketInCurArcIndex = 0; PacketInCurArcIndex < PacketCountInCurArc; ++PacketInCurArcIndex)
 		{
+			// PacketCount was determined from the sum of GetPacketCount() values earlier,
+			// so we should be in bounds unless something went badly wrong.
+			check(PacketIdx < (uint32)PacketCount);
 			uint32 PacketSize = BufferSize;
 
 			CurArc->SerializePacket(ReadBuffer, PacketSize);
